Basic TapJoy ad support working

diff --git a/Engine/Plugins/Runtime/Advertising/IOSTapJoy/Source/IOSTapJoy/Private/IOSTapJoy.cpp b/Engine/Plugins/Runtime/Advertising/IOSTapJoy/Source/IOSTapJoy/Private/IOSTapJoy.cpp
--- a/Engine/Plugins/Runtime/Advertising/IOSTapJoy/Source/IOSTapJoy/Private/IOSTapJoy.cpp
+++ b/Engine/Plugins/Runtime/Advertising/IOSTapJoy/Source/IOSTapJoy/Private/IOSTapJoy.cpp
@@ -1,7 +1,7 @@
 // Copyright 1998-2014 Epic Games, Inc. All Rights Reserved.
 
 #include "IOSTapJoyPrivatePCH.h"
-
+#include "IOSAppDelegate.h"
 #import "../ThirdPartyFrameworks/Tapjoy.embeddedframework/Tapjoy.framework/Headers/Tapjoy.h"
 
 class FTapJoyProvider : public IAdvertisingProvider
@@ -18,10 +18,17 @@ class FTapJoyProvider : public IAdvertisingProvider
 
 IMPLEMENT_MODULE( FTapJoyProvider, IOSTapJoy )
 
-@interface IOSTapJoy : UIResponder
+static const FString AppIDString( TEXT( "" ) );
+static const FString SecretKeyString( TEXT( "" ) );
+static const FString CurrencyString( TEXT( "" ) );
+
+@interface IOSTapJoy : UIResponder< TJCAdDelegate >
 @end
 
 @implementation IOSTapJoy
+
+static bool bAttemptingToShowAd	= false;
+
 + (IOSTapJoy*)GetDelegate
 {
 	static IOSTapJoy * Singleton = [[IOSTapJoy alloc] init];
@@ -53,8 +60,8 @@ IMPLEMENT_MODULE( FTapJoyProvider, IOSTapJoy )
 		object:nil];
 
 	// This code will execute after your module is loaded into memory (but after global variables are initialized, of course.)
-	[Tapjoy requestTapjoyConnect:@"AppID"
-		secretKey:@"SecretKey"
+	[Tapjoy requestTapjoyConnect:[NSString stringWithFString:AppIDString]
+		secretKey:[NSString stringWithFString:SecretKeyString]
 		options:@{ TJC_OPTION_ENABLE_LOGGING : @(YES) }
 	// If you are not using Tapjoy Managed currency, you would set your own user ID here.
 	// TJC_OPTION_USER_ID : @"A_UNIQUE_USER_ID"
@@ -64,6 +71,127 @@ IMPLEMENT_MODULE( FTapJoyProvider, IOSTapJoy )
 	// TJC_OPTION_SEGMENTATION_PARAMS : segmentationParams
 	];
 }
+
+-(void)ShowAdBanner:(NSNumber*)bShowOnBottomOfScreen
+{
+	if ( bAttemptingToShowAd )
+	{
+		NSLog( @"ShowAdBanner: Already attempting to show ad..." );
+		return;
+	}
+
+	if ( !CurrencyString.IsEmpty() )
+	{
+		NSString * CurrencyStringNS = [NSString stringWithFString:CurrencyString];
+
+		//NSLog( @"ShowAdBanner: %@", CurrencyStringNS );
+
+		[Tapjoy getDisplayAdWithDelegate:self currencyID:CurrencyStringNS];
+	}
+	else
+	{
+		[Tapjoy getDisplayAdWithDelegate:self];
+	}
+
+	bAttemptingToShowAd = true;
+}
+
+-(void)HideAdBanner
+{
+	TJCAdView * adView = nil;
+
+	for ( UIView * SubView in [[IOSAppDelegate GetDelegate].RootView subviews] ) 
+	{
+		if ( [SubView isKindOfClass:[TJCAdView class]] ) 
+		{
+			adView = (TJCAdView*)SubView;
+			break;
+		}
+	}
+
+	if ( adView == nil )
+	{
+		NSLog( @"HideAdBanner: No ad view is active..." );
+		return;
+	}
+
+	// Fade it out
+	//if ( ![Tapjoy getDisplayAdView].hidden )
+	if ( !adView.hidden )
+	{
+		[UIView animateWithDuration:0.4f 
+			animations:^
+			{ 
+				adView.alpha = 0.0f; 
+			} 
+			completion:^(BOOL finished) 
+			{ 
+				adView.hidden = YES; 
+			}
+		];
+	}
+}
+
+- (void)didReceiveAd:(TJCAdView*)adView
+{
+	NSLog( @"didReceiveAd called..." );
+
+	if ( !bAttemptingToShowAd )
+	{
+		NSLog( @"didReceiveAd: bAttemptingToShowAd == false" );
+	}
+
+	// Remove any existing ad views
+	for ( UIView * SubView in [[IOSAppDelegate GetDelegate].RootView subviews] ) 
+	{
+		if ( [SubView isKindOfClass:[TJCAdView class]] ) 
+		{
+			[SubView removeFromSuperview];
+		}
+	}
+
+	// Hide it initially
+	adView.hidden = YES;
+	adView.alpha = 0.0f;
+
+	// Add it it the main view
+	[[IOSAppDelegate GetDelegate].RootView addSubview : adView];
+
+	// Fade it in
+	if ( adView.hidden )
+	{
+		adView.hidden = NO;
+		[UIView animateWithDuration:0.4f
+			animations:^
+			{
+				adView.alpha = 1.0f;
+			}
+		];
+	}
+
+	bAttemptingToShowAd = false;
+}
+
+- (void)didFailWithMessage:(NSString*)msg
+{
+	if ( !bAttemptingToShowAd )
+	{
+		NSLog( @"didFailWithMessage: bAttemptingToShowAd == false" );
+	}
+
+	NSLog( @"didFailWithMessage: %@", msg );
+	bAttemptingToShowAd = false;
+}
+
+- (NSString*)adContentSize
+{
+	return TJC_DISPLAY_AD_SIZE_320X50;
+}
+
+- (BOOL)shouldRefreshAd
+{
+	return NO;
+}
 @end
 
 void FTapJoyProvider::StartupModule() 
@@ -77,14 +205,15 @@ void FTapJoyProvider::ShutdownModule()
 
 void FTapJoyProvider::ShowAdBanner( bool bShowOnBottomOfScreen )
 {
+	[[IOSTapJoy GetDelegate] performSelectorOnMainThread:@selector(ShowAdBanner:) withObject:[NSNumber numberWithBool : bShowOnBottomOfScreen] waitUntilDone : NO];
 }
 
 void FTapJoyProvider::HideAdBanner()
 {
-
+	[[IOSTapJoy GetDelegate] performSelectorOnMainThread:@selector(HideAdBanner) withObject:nil waitUntilDone : NO];
 }
 
 void FTapJoyProvider::CloseAdBanner()
 {
-
+	HideAdBanner();
 }
