Pull TapJoy settings from config

diff --git a/Engine/Plugins/Runtime/Advertising/IOSTapJoy/Source/IOSTapJoy/Private/IOSTapJoy.cpp b/Engine/Plugins/Runtime/Advertising/IOSTapJoy/Source/IOSTapJoy/Private/IOSTapJoy.cpp
--- a/Engine/Plugins/Runtime/Advertising/IOSTapJoy/Source/IOSTapJoy/Private/IOSTapJoy.cpp
+++ b/Engine/Plugins/Runtime/Advertising/IOSTapJoy/Source/IOSTapJoy/Private/IOSTapJoy.cpp
@@ -18,9 +18,9 @@ class FTapJoyProvider : public IAdvertisingProvider
 
 IMPLEMENT_MODULE( FTapJoyProvider, IOSTapJoy )
 
-static const FString AppIDString( TEXT( "" ) );
-static const FString SecretKeyString( TEXT( "" ) );
-static const FString CurrencyString( TEXT( "" ) );
+static FString AppIDString( TEXT( "" ) );
+static FString SecretKeyString( TEXT( "" ) );
+static FString CurrencyString( TEXT( "" ) );
 
 @interface IOSTapJoy : UIResponder< TJCAdDelegate >
 @end
@@ -196,6 +196,10 @@ static bool bAttemptingToShowAd	= false;
 
 void FTapJoyProvider::StartupModule() 
 {
+	GConfig->GetString( TEXT( "TapJoy" ), TEXT( "AppID" ),			AppIDString,		GEngineIni );
+	GConfig->GetString( TEXT( "TapJoy" ), TEXT( "SecretKey" ),		SecretKeyString,	GEngineIni );
+	GConfig->GetString( TEXT( "TapJoy") , TEXT( "CurrencyString" ),	CurrencyString,		GEngineIni );
+
 	[[IOSTapJoy GetDelegate] performSelectorOnMainThread:@selector(StartupTapJoy) withObject:nil waitUntilDone : NO];
 }
 
