Fixes for building the TapJoy plugin when it's enabled

diff --git a/Engine/Plugins/Runtime/Advertising/IOSTapJoy/Source/IOSTapJoy/Private/IOSTapJoy.cpp b/Engine/Plugins/Runtime/Advertising/IOSTapJoy/Source/IOSTapJoy/Private/IOSTapJoy.cpp
--- a/Engine/Plugins/Runtime/Advertising/IOSTapJoy/Source/IOSTapJoy/Private/IOSTapJoy.cpp
+++ b/Engine/Plugins/Runtime/Advertising/IOSTapJoy/Source/IOSTapJoy/Private/IOSTapJoy.cpp
@@ -2,7 +2,7 @@
 
 #include "IOSTapJoyPrivatePCH.h"
 #include "IOSAppDelegate.h"
-#import "../ThirdPartyFrameworks/Tapjoy.embeddedframework/Tapjoy.framework/Headers/Tapjoy.h"
+#import <Tapjoy/Tapjoy.h>
 
 class FTapJoyProvider : public IAdvertisingProvider
 {
diff --git a/Engine/Source/Programs/UnrealBuildTool/IOS/IOSToolChain.cs b/Engine/Source/Programs/UnrealBuildTool/IOS/IOSToolChain.cs
--- a/Engine/Source/Programs/UnrealBuildTool/IOS/IOSToolChain.cs
+++ b/Engine/Source/Programs/UnrealBuildTool/IOS/IOSToolChain.cs
@@ -1415,7 +1415,7 @@ namespace UnrealBuildTool
 							continue;		// This framework item doesn't belong to this target, skip it
 						}
 
-						string RemoteZipPath = GetRemoteFrameworkZipPath( Framework );
+						string RemoteZipPath = GetRemoteIntermediateFrameworkZipPath( Framework );
 
 						RemoteZipPath = RemoteZipPath.Replace( ".zip", "" );
 
