- Fixing CIS error (hopefully) in a plugin that I can't repro, and we didn't see in Preflight or integration testing :(

diff --git a/Engine/Plugins/Runtime/Advertising/IOSTapJoy/Source/IOSTapJoy/Private/IOSTapJoy.cpp b/Engine/Plugins/Runtime/Advertising/IOSTapJoy/Source/IOSTapJoy/Private/IOSTapJoy.cpp
--- a/Engine/Plugins/Runtime/Advertising/IOSTapJoy/Source/IOSTapJoy/Private/IOSTapJoy.cpp
+++ b/Engine/Plugins/Runtime/Advertising/IOSTapJoy/Source/IOSTapJoy/Private/IOSTapJoy.cpp
@@ -15,6 +15,11 @@ class FTapJoyProvider : public IAdvertisingProvider
 	virtual void HideAdBanner() override;
 	virtual void CloseAdBanner() override;
 	virtual int32 GetAdIDCount() override;
+
+	virtual void LoadInterstitialAd(int32 adID) override;
+	virtual bool IsInterstitialAdAvailable() override;
+	virtual bool IsInterstitialAdRequested() override;
+	virtual void ShowInterstitialAd() override;
 };
 
 IMPLEMENT_MODULE( FTapJoyProvider, IOSTapJoy )
@@ -227,3 +232,24 @@ int32 FTapJoyProvider::GetAdIDCount()
 {
 	return 1;
 }
+
+void LoadInterstitialAd(int32 adID)
+{
+
+}
+
+bool IsInterstitialAdAvailable()
+{
+	return false;
+}
+
+bool IsInterstitialAdRequested()
+{
+	return false;
+}
+
+void ShowInterstitialAd()
+{
+	
+}
+
