[TapJoy]
* Fixed casing for framework lib. It is lowercase j

#Synced-CL 50216544

diff --git a/Engine/Plugins/Runtime/Advertising/IOSTapJoy/Source/IOSTapJoy/IOSTapJoy.Build.cs b/Engine/Plugins/Runtime/Advertising/IOSTapJoy/Source/IOSTapJoy/IOSTapJoy.Build.cs
--- a/Engine/Plugins/Runtime/Advertising/IOSTapJoy/Source/IOSTapJoy/IOSTapJoy.Build.cs
+++ b/Engine/Plugins/Runtime/Advertising/IOSTapJoy/Source/IOSTapJoy/IOSTapJoy.Build.cs
@@ -6,6 +6,8 @@ namespace UnrealBuildTool.Rules
 	{
 		public IOSTapJoy( ReadOnlyTargetRules Target ) : base(Target)
 		{
+			bRequiresPlatformSDK = true;
+		
 			PublicDependencyModuleNames.AddRange(
 				new string[]
 				{
@@ -25,10 +27,10 @@ namespace UnrealBuildTool.Rules
 
 			PublicIncludePathModuleNames.Add( "Advertising" );
 
-			// Add the TapJoy framework
+			// Add the Tapjoy framework
 			PublicAdditionalFrameworks.Add( 
 				new Framework( 
-					"TapJoy",														// Framework name
+					"Tapjoy",														// Framework name
 					"../../ThirdPartyFrameworks/Tapjoy.embeddedframework.zip",		// Zip name
 					"Resources/TapjoyResources.bundle"								// Resources we need copied and staged
 				)
